Updates actionable warnings to be aggregated and collapsed by default.

diff --git a/Engine/Source/Editor/ActionableMessage/Private/ActionableMessageSubsystem.cpp b/Engine/Source/Editor/ActionableMessage/Private/ActionableMessageSubsystem.cpp
--- a/Engine/Source/Editor/ActionableMessage/Private/ActionableMessageSubsystem.cpp
+++ b/Engine/Source/Editor/ActionableMessage/Private/ActionableMessageSubsystem.cpp
@@ -29,6 +29,8 @@ void UActionableMessageSubsystem::SetActionableMessage(FName InProvider, const F
 
 void UActionableMessageSubsystem::ClearActionableMessage(FName InProvider)
 {
-	ProviderActionableMessageMap.Remove(InProvider);
-	++StateID;
+	if (ProviderActionableMessageMap.Remove(InProvider) > 0)
+	{
+		++StateID;
+	}
 }
diff --git a/Engine/Source/Editor/ActionableMessage/Private/SActionableMessageViewportWidget.cpp b/Engine/Source/Editor/ActionableMessage/Private/SActionableMessageViewportWidget.cpp
--- a/Engine/Source/Editor/ActionableMessage/Private/SActionableMessageViewportWidget.cpp
+++ b/Engine/Source/Editor/ActionableMessage/Private/SActionableMessageViewportWidget.cpp
@@ -23,91 +23,67 @@
 void SActionableMessageEntry::Construct(const FArguments& InArgs)
 {
 	ActionableMessage = InArgs._ActionableMessage;
-	OnClicked = InArgs._OnClicked;
 	
 	ChildSlot
 	[
-		SNew(SButton)
-		.ButtonStyle(&FCoreStyle::Get().GetWidgetStyle<FButtonStyle>("HoverHintOnly"))
-		.IsEnabled(true)
-		.OnClicked_Lambda([this]()
-		{
-			if (OnClicked.IsBound())
-			{
-				return OnClicked.Execute();
-			}
-
-			return FReply::Handled();
-		})
+		SNew(SHorizontalBox)
 		.ToolTipText_Lambda([this]() { return ActionableMessage.IsValid() ? ActionableMessage->Tooltip : FText::GetEmpty(); })
-		.ContentPadding(FMargin(0.f, 0.f, 0.f, -1.f))
+		+ SHorizontalBox::Slot()
+		.Padding(FMargin(4, 0, 8, 0))
+		.VAlign(VAlign_Center)
+		.HAlign(HAlign_Left)
+		.FillWidth(1)
 		[
-			SNew(SHorizontalBox)
-			+ SHorizontalBox::Slot()
-			.AutoWidth()
-			.VAlign(VAlign_Center)
-			[
-				SNew(SImage)
-				.Image(FAppStyle::Get().GetBrush("ActionableMessage.Warning"))
-				.DesiredSizeOverride(FVector2D(16, 16))
-			]
-			+ SHorizontalBox::Slot()
-			.Padding(FMargin(4, 0, 8, 0))
-			.VAlign(VAlign_Center)
-			.HAlign(HAlign_Left)
-			.FillWidth(1)
-			[
-				SNew(STextBlock)
-				.Text_Lambda([this]() { return ActionableMessage.IsValid() ? ActionableMessage->Message : FText::GetEmpty(); })
-				.ColorAndOpacity(FStyleColors::Foreground)
-			]
-			+ SHorizontalBox::Slot()
-			.VAlign(VAlign_Center)
-			.HAlign(HAlign_Right)
-			.AutoWidth()
-			[
-				SNew(SButton)
-				.Visibility_Lambda([this]()
+			SNew(STextBlock)
+			.Text_Lambda([this]() { return ActionableMessage.IsValid() ? ActionableMessage->Message : FText::GetEmpty(); })
+			.ColorAndOpacity(FStyleColors::Foreground)
+		]
+		+ SHorizontalBox::Slot()
+		.VAlign(VAlign_Center)
+		.HAlign(HAlign_Right)
+		.AutoWidth()
+		[
+			SNew(SButton)
+			.Visibility_Lambda([this]()
+			{
+				if (!ActionableMessage.IsValid() || ActionableMessage->ActionMessage.IsEmpty())
 				{
-					if (!ActionableMessage.IsValid() || ActionableMessage->ActionMessage.IsEmpty())
-					{
-						return EVisibility::Collapsed;
-					}
-
-					return EVisibility::Visible;
-				})
-				.ButtonStyle(&FCoreStyle::Get().GetWidgetStyle<FButtonStyle>("Button"))
-				.ContentPadding(FMargin(-4.f, 2.f))
-				.OnClicked_Lambda([this]()
+					return EVisibility::Collapsed;
+				}
+
+				return EVisibility::Visible;
+			})
+			.ButtonStyle(&FCoreStyle::Get().GetWidgetStyle<FButtonStyle>("Button"))
+			.ContentPadding(FMargin(-4.f, 2.f))
+			.OnClicked_Lambda([this]()
+			{
+				if (ActionableMessage.IsValid() && (ActionableMessage->ActionCallback != nullptr))
 				{
-					if (ActionableMessage.IsValid() && (ActionableMessage->ActionCallback != nullptr))
-					{
-						ActionableMessage->ActionCallback();
-					}
+					ActionableMessage->ActionCallback();
+				}
 
-					return FReply::Handled();
-				})
+				return FReply::Handled();
+			})
+			[
+				SNew(SHorizontalBox)
+				+ SHorizontalBox::Slot()
+				.AutoWidth()
+				.Padding(0.f, 0.f, 4.f, 0.f)
+				.VAlign(VAlign_Center)
+				.HAlign(HAlign_Left)
 				[
-					SNew(SHorizontalBox)
-					+ SHorizontalBox::Slot()
-					.AutoWidth()
-					.Padding(0.f, 0.f, 4.f, 0.f)
-					.VAlign(VAlign_Center)
-					.HAlign(HAlign_Left)
-					[
-						SNew(SImage)
-						.Image(FAppStyle::Get().GetBrush("ActionableMessage.Update"))
-						.DesiredSizeOverride(FVector2D(12, 12))
-					]
-					+ SHorizontalBox::Slot()
-					.AutoWidth()
-					.VAlign(VAlign_Center)
-					.HAlign(HAlign_Right)
-					[
-						SNew(STextBlock)
-						.Text_Lambda([this]() { return ActionableMessage.IsValid() ? ActionableMessage->ActionMessage : FText::GetEmpty(); })
-						.ColorAndOpacity(FStyleColors::White)
-					]
+					SNew(SImage)
+					.Image(FAppStyle::Get().GetBrush("ActionableMessage.Update"))
+					.DesiredSizeOverride(FVector2D(12, 12))
+				]
+				+ SHorizontalBox::Slot()
+				.AutoWidth()
+				.VAlign(VAlign_Center)
+				.HAlign(HAlign_Right)
+				[
+					SNew(STextBlock)
+					.Text_Lambda([this]() { return ActionableMessage.IsValid() ? ActionableMessage->ActionMessage : FText::GetEmpty(); })
+					.ColorAndOpacity(FStyleColors::White)
 				]
 			]
 		]
@@ -128,21 +104,49 @@ void SActionableMessageViewportWidget::Construct(const FArguments& InArgs)
 		.HAlign(HAlign_Right)
 		.VAlign(VAlign_Top)
 		[
-			SNew(SImage)
-			.Image_Lambda([this]() { return bExpanded ? FAppStyle::GetBrush("ContentBrowser.SortDown") : FAppStyle::GetBrush("ContentBrowser.SortUp"); })
-			.Visibility_Lambda([this]() { return ActionableMessages.IsEmpty() ? EVisibility::Collapsed : EVisibility::Visible; })
-			.ColorAndOpacity(FSlateColor::UseForeground())
-		]
-		+ SVerticalBox::Slot()
-		.Padding(8.f, 4.f)
-		.AutoHeight()
-		[
-			SAssignNew(MainActionableMessage, SActionableMessageEntry)
+			SNew(SButton)
+			.ButtonStyle(&FCoreStyle::Get().GetWidgetStyle<FButtonStyle>("HoverHintOnly"))
+			.IsEnabled(true)
 			.OnClicked_Lambda([this]()
 			{
 				bExpanded = !bExpanded;
 				return FReply::Handled();
-			})				
+			})
+			.ContentPadding(FMargin(0.f, 4.f, 0.f, 4.f))
+			[
+				SNew(SHorizontalBox)
+				+ SHorizontalBox::Slot()
+				.Padding(FMargin(4.f, 0, 4.f, 0))
+				.VAlign(VAlign_Center)
+				.HAlign(HAlign_Left)
+				.FillWidth(1)
+				[
+					SNew(SImage)
+					.Image_Lambda([this]() { return bExpanded ? FAppStyle::GetBrush("ContentBrowser.SortDown") : FAppStyle::GetBrush("Symbols.RightArrow"); })
+					.Visibility_Lambda([this]() { return ActionableMessages.IsEmpty() ? EVisibility::Collapsed : EVisibility::Visible; })
+					.ColorAndOpacity(FSlateColor::UseForeground())
+				]
+				+ SHorizontalBox::Slot()
+				.Padding(FMargin(0.f, 0, 2.f, 0))
+				.VAlign(VAlign_Center)
+				.HAlign(HAlign_Center)
+				.FillWidth(1)
+				[
+					SNew(STextBlock)
+					.Text_Lambda([this]() { return FText::FromString(FString::FromInt(ActionableMessages.Num())); })
+					.Visibility_Lambda([this]() { return (ActionableMessages.Num() > 1) ? EVisibility::Visible : EVisibility::Collapsed; })
+					.ColorAndOpacity(FStyleColors::Foreground)
+				]
+				+ SHorizontalBox::Slot()
+				.AutoWidth()
+				.VAlign(VAlign_Center)
+				.Padding(FMargin(0.f, 0, 4.f, 0))
+				[
+					SNew(SImage)
+					.Image(FAppStyle::Get().GetBrush("ActionableMessage.Warning"))
+					.DesiredSizeOverride(FVector2D(16, 16))
+				]
+			]
 		]
 		+ SVerticalBox::Slot()
 		.AutoHeight()
@@ -175,30 +179,21 @@ EVisibility SActionableMessageViewportWidget::GetVisibility()
 
 		if (CachedStateID == SubsystemStateID)
 		{
-			return bMainMessageFilled ? EVisibility::Visible : EVisibility::Collapsed;
+			return ActionableMessages.IsEmpty() ? EVisibility::Collapsed : EVisibility::Visible;
 		}
 		
 		const TMap<FName, TSharedPtr<FActionableMessage>>& ActionableMessageSource = ActionableMessageSubsystem->GetActionableMessages();
 
-		bMainMessageFilled = false;
 		ActionableMessages.Empty();
-		
-		for (const TPair<FName, TSharedPtr<FActionableMessage>>& ActionableMessagePair : ActionableMessageSource)
+
+		Algo::Transform(ActionableMessageSource, ActionableMessages, [](const TPair<FName, TSharedPtr<FActionableMessage>>& ActionableMessagePair)
 		{
-			if ((!bMainMessageFilled) && (MainActionableMessage.IsValid()))
-			{
-				bMainMessageFilled = true;
-				MainActionableMessage->SetActionableMessage(ActionableMessagePair.Value);
-			}
-			else
-			{
-				ActionableMessages.Emplace(ActionableMessagePair.Value);
-			}
-		}
+			return ActionableMessagePair.Value;
+		});
 
 		CachedStateID = SubsystemStateID;
 		
-		return bMainMessageFilled ? EVisibility::Visible : EVisibility::Collapsed;
+		return ActionableMessages.IsEmpty() ? EVisibility::Collapsed : EVisibility::Visible;
 	}
 
 	return EVisibility::Collapsed;
diff --git a/Engine/Source/Editor/ActionableMessage/Public/SActionableMessageViewportWidget.h b/Engine/Source/Editor/ActionableMessage/Public/SActionableMessageViewportWidget.h
--- a/Engine/Source/Editor/ActionableMessage/Public/SActionableMessageViewportWidget.h
+++ b/Engine/Source/Editor/ActionableMessage/Public/SActionableMessageViewportWidget.h
@@ -14,7 +14,6 @@ class ACTIONABLEMESSAGE_API SActionableMessageEntry : public SCompoundWidget
 public:
 	SLATE_BEGIN_ARGS(SActionableMessageEntry) {}
 	SLATE_ARGUMENT(TSharedPtr<FActionableMessage>, ActionableMessage)
-	SLATE_EVENT(FOnClicked, OnClicked)
 	SLATE_END_ARGS()
 
 	void Construct(const FArguments& InArgs);
@@ -37,10 +36,8 @@ public:
 	TSharedRef<ITableRow> OnGenerateRow(TSharedPtr<FActionableMessage> InActionableMessage, const TSharedRef<STableViewBase>& OwnerTable);
 
 private:
-	TSharedPtr<SActionableMessageEntry> MainActionableMessage;
 	TSharedPtr<SListView<TSharedPtr<FActionableMessage>>> ActionableMessageList;
 	TArray<TSharedPtr<FActionableMessage>> ActionableMessages;
 	uint32 CachedStateID = 0;
-	bool bMainMessageFilled = false;
 	bool bExpanded = false;
 };
