Updates styles for actionable messages.

diff --git a/Engine/Source/Editor/ActionableMessage/Private/SActionableMessageViewportWidget.cpp b/Engine/Source/Editor/ActionableMessage/Private/SActionableMessageViewportWidget.cpp
--- a/Engine/Source/Editor/ActionableMessage/Private/SActionableMessageViewportWidget.cpp
+++ b/Engine/Source/Editor/ActionableMessage/Private/SActionableMessageViewportWidget.cpp
@@ -10,6 +10,8 @@
 #include "Styling/AppStyle.h"
 #include "Styling/StyleColors.h"
 #include "Modules/ModuleManager.h"
+#include "Widgets/Layout/SSeparator.h"
+#include "Widgets/Layout/SSpacer.h"
 #include "Widgets/Input/SButton.h"
 #include "Widgets/Images/SImage.h"
 #include "Widgets/Text/STextBlock.h"
@@ -29,7 +31,17 @@ void SActionableMessageEntry::Construct(const FArguments& InArgs)
 		SNew(SHorizontalBox)
 		.ToolTipText_Lambda([this]() { return ActionableMessage.IsValid() ? ActionableMessage->Tooltip : FText::GetEmpty(); })
 		+ SHorizontalBox::Slot()
-		.Padding(FMargin(4, 0, 8, 0))
+		.VAlign(VAlign_Center)
+		.HAlign(HAlign_Left)
+		.Padding(FMargin(0.f, 0.f, 4.f, 0.f))
+		.AutoWidth()
+		[
+			SNew(SImage)
+			.Image(FAppStyle::Get().GetBrush("ActionableMessage.Warning"))
+			.DesiredSizeOverride(FVector2D(16, 16))
+		]
+		+ SHorizontalBox::Slot()
+		.Padding(FMargin(0.f, 0.f, 4.f, 0.f))
 		.VAlign(VAlign_Center)
 		.HAlign(HAlign_Left)
 		.FillWidth(1)
@@ -97,11 +109,14 @@ void SActionableMessageEntry::SetActionableMessage(TSharedPtr<FActionableMessage
 
 void SActionableMessageViewportWidget::Construct(const FArguments& InArgs)
 {
+	constexpr float IconSize = 16.f;
+	constexpr float DefaultPadding = 4.f;
+
 	ChildSlot
 	[
 		SNew(SVerticalBox)
 		+ SVerticalBox::Slot()
-		.HAlign(HAlign_Right)
+		.HAlign(HAlign_Left)
 		.VAlign(VAlign_Top)
 		[
 			SNew(SButton)
@@ -112,43 +127,88 @@ void SActionableMessageViewportWidget::Construct(const FArguments& InArgs)
 				bExpanded = !bExpanded;
 				return FReply::Handled();
 			})
-			.ContentPadding(FMargin(0.f, 4.f, 0.f, 4.f))
+			.ContentPadding(FMargin(DefaultPadding))
 			[
 				SNew(SHorizontalBox)
 				+ SHorizontalBox::Slot()
-				.Padding(FMargin(4.f, 0, 4.f, 0))
 				.VAlign(VAlign_Center)
 				.HAlign(HAlign_Left)
-				.FillWidth(1)
+				.Padding(FMargin(0.f, 0.f, DefaultPadding, 0.f))
+				.AutoWidth()
 				[
 					SNew(SImage)
-					.Image_Lambda([this]() { return bExpanded ? FAppStyle::GetBrush("ContentBrowser.SortDown") : FAppStyle::GetBrush("Symbols.RightArrow"); })
-					.Visibility_Lambda([this]() { return ActionableMessages.IsEmpty() ? EVisibility::Collapsed : EVisibility::Visible; })
-					.ColorAndOpacity(FSlateColor::UseForeground())
+					.Image(FAppStyle::Get().GetBrush("ActionableMessage.Warning"))
+					.DesiredSizeOverride(FVector2D(IconSize, IconSize))
 				]
 				+ SHorizontalBox::Slot()
-				.Padding(FMargin(0.f, 0, 2.f, 0))
 				.VAlign(VAlign_Center)
-				.HAlign(HAlign_Center)
-				.FillWidth(1)
+				.HAlign(HAlign_Left)
+				.AutoWidth()
 				[
-					SNew(STextBlock)
-					.Text_Lambda([this]() { return FText::FromString(FString::FromInt(ActionableMessages.Num())); })
-					.Visibility_Lambda([this]() { return (ActionableMessages.Num() > 1) ? EVisibility::Visible : EVisibility::Collapsed; })
+					SAssignNew(TextBlock, STextBlock)
+					.Text_Lambda([this]()
+					{
+						if (bExpanded)
+						{
+							if (ActionableMessages.Num() == 1)
+							{
+								return LOCTEXT("ActionableMessages.Warning", "Warning");
+							}
+							
+							return FText::Format(LOCTEXT("ActionableMessages.Number", "{0} Warnings"), ActionableMessages.Num());
+						}
+
+						return FText::FromString(FString::FromInt(ActionableMessages.Num()));
+					})
 					.ColorAndOpacity(FStyleColors::Foreground)
+					.Font_Lambda([this] ()
+					{
+						if (bExpanded)
+						{
+							return FAppStyle::GetFontStyle(TEXT("BoldFont")); 
+						}
+
+						return FAppStyle::GetFontStyle(TEXT("NormalFont"));
+					})
 				]
 				+ SHorizontalBox::Slot()
-				.AutoWidth()
+				.FillWidth(1)
+				[
+					SNew(SSpacer)
+					.Size_Lambda([this]()
+					{
+						FVector2d Size(0.f, 0.f);
+						
+						if (bExpanded && ActionableMessageList.IsValid() && TextBlock.IsValid())
+						{
+							Size.X = ActionableMessageList->GetDesiredSize().X - TextBlock->GetDesiredSize().X - IconSize - 4.f * DefaultPadding - 2.f;
+						}
+		
+						return Size;
+					})
+				]
+				+ SHorizontalBox::Slot()
+				.Padding(FMargin(DefaultPadding, 0.f, 0.f, 0.f))
 				.VAlign(VAlign_Center)
-				.Padding(FMargin(0.f, 0, 4.f, 0))
+				.HAlign(HAlign_Right)
+				.AutoWidth()
 				[
 					SNew(SImage)
-					.Image(FAppStyle::Get().GetBrush("ActionableMessage.Warning"))
-					.DesiredSizeOverride(FVector2D(16, 16))
+					.Image_Lambda([this]() { return bExpanded ? FAppStyle::GetBrush("ContentBrowser.SortDown") : FAppStyle::GetBrush("Symbols.LeftArrow"); })
+					.Visibility_Lambda([this]() { return ActionableMessages.IsEmpty() ? EVisibility::Collapsed : EVisibility::Visible; })
+					.ColorAndOpacity(FSlateColor::UseForeground())
 				]
 			]
 		]
 		+ SVerticalBox::Slot()
+		.Padding(0.f, 0.f, 0.f, DefaultPadding)
+		.AutoHeight()
+		[
+			SNew(SSeparator)
+			.Thickness(1.f)
+			.Visibility_Lambda([this]() {return bExpanded ? EVisibility::Visible : EVisibility::Collapsed;})
+		]
+		+ SVerticalBox::Slot()
 		.AutoHeight()
 		[
 			SAssignNew(ActionableMessageList, SListView<TSharedPtr<FActionableMessage>>)
@@ -205,7 +265,7 @@ TSharedRef<ITableRow> SActionableMessageViewportWidget::OnGenerateRow(TSharedPtr
 
 	return SNew(STableRow<TSharedPtr<FActionableMessage>>, OwnerTable)
 	.Style(&FAppStyle::Get().GetWidgetStyle<FTableRowStyle>("ActionableMessage.ListViewRow"))
-	.Padding(FMargin(8.f, 0.f, 8.f, 4.f))
+	.Padding(FMargin(4.f, 0.f, 4.f, 4.f))
 	[
 		SNew(SActionableMessageEntry)
 		.ActionableMessage(InActionableMessage)
diff --git a/Engine/Source/Editor/ActionableMessage/Public/SActionableMessageViewportWidget.h b/Engine/Source/Editor/ActionableMessage/Public/SActionableMessageViewportWidget.h
--- a/Engine/Source/Editor/ActionableMessage/Public/SActionableMessageViewportWidget.h
+++ b/Engine/Source/Editor/ActionableMessage/Public/SActionableMessageViewportWidget.h
@@ -37,6 +37,7 @@ public:
 
 private:
 	TSharedPtr<SListView<TSharedPtr<FActionableMessage>>> ActionableMessageList;
+	TSharedPtr<STextBlock> TextBlock;
 	TArray<TSharedPtr<FActionableMessage>> ActionableMessages;
 	uint32 CachedStateID = 0;
 	bool bExpanded = false;
diff --git a/Engine/Source/Editor/EditorStyle/Private/StarshipStyle.cpp b/Engine/Source/Editor/EditorStyle/Private/StarshipStyle.cpp
--- a/Engine/Source/Editor/EditorStyle/Private/StarshipStyle.cpp
+++ b/Engine/Source/Editor/EditorStyle/Private/StarshipStyle.cpp
@@ -400,6 +400,7 @@ void FStarshipEditorStyle::FStyle::SetupGeneralStyles()
 		Set( "Symbols.DownArrow", new IMAGE_BRUSH( "Common/DownArrow", Icon8x8 ) );
 		Set( "Symbols.DoubleDownArrow", new IMAGE_BRUSH( "Common/DownArrow2", Icon8x8 ) );
 		Set( "Symbols.RightArrow", new IMAGE_BRUSH("Common/SubmenuArrow", Icon8x8));
+		Set( "Symbols.LeftArrow", new IMAGE_BRUSH("Common/LeftArrow", Icon8x8));
 		Set( "Symbols.Check", new IMAGE_BRUSH( "Common/Check", Icon16x16 ) );
 	}
 
