Fixes actionable messages not updating properly when the same provider pushes different messages

diff --git a/Engine/Source/Editor/ActionableMessage/Private/SActionableMessageViewportWidget.cpp b/Engine/Source/Editor/ActionableMessage/Private/SActionableMessageViewportWidget.cpp
--- a/Engine/Source/Editor/ActionableMessage/Private/SActionableMessageViewportWidget.cpp
+++ b/Engine/Source/Editor/ActionableMessage/Private/SActionableMessageViewportWidget.cpp
@@ -252,6 +252,11 @@ EVisibility SActionableMessageViewportWidget::GetVisibility()
 		});
 
 		CachedStateID = SubsystemStateID;
+
+		if (ActionableMessageList.IsValid())
+		{
+			ActionableMessageList->RebuildList();
+		}
 		
 		return ActionableMessages.IsEmpty() ? EVisibility::Collapsed : EVisibility::Visible;
 	}
