[LevelViewportToolBar] Fix toolbar widget hit test visibilities to allow cursor tests to pass through. Fixes cursor icon issues with new Sequencer select-ability feature.

diff --git a/Engine/Source/Editor/ActionableMessage/Private/SActionableMessageViewportWidget.cpp b/Engine/Source/Editor/ActionableMessage/Private/SActionableMessageViewportWidget.cpp
--- a/Engine/Source/Editor/ActionableMessage/Private/SActionableMessageViewportWidget.cpp
+++ b/Engine/Source/Editor/ActionableMessage/Private/SActionableMessageViewportWidget.cpp
@@ -247,7 +247,7 @@ EVisibility SActionableMessageViewportWidget::GetVisibility()
 
 		if (CachedStateID == SubsystemStateID)
 		{
-			return ActionableMessages.IsEmpty() ? EVisibility::Collapsed : EVisibility::Visible;
+			return ActionableMessages.IsEmpty() ? EVisibility::Collapsed : EVisibility::HitTestInvisible;
 		}
 		
 		const TMap<FName, TSharedPtr<FActionableMessage>>& ActionableMessageSource = ActionableMessageSubsystem->GetActionableMessages();
@@ -266,7 +266,7 @@ EVisibility SActionableMessageViewportWidget::GetVisibility()
 			ActionableMessageList->RebuildList();
 		}
 		
-		return ActionableMessages.IsEmpty() ? EVisibility::Collapsed : EVisibility::Visible;
+		return ActionableMessages.IsEmpty() ? EVisibility::Collapsed : EVisibility::HitTestInvisible;
 	}
 
 	return EVisibility::Collapsed;
diff --git a/Engine/Source/Editor/LevelEditor/Private/SLevelViewportToolBar.cpp b/Engine/Source/Editor/LevelEditor/Private/SLevelViewportToolBar.cpp
--- a/Engine/Source/Editor/LevelEditor/Private/SLevelViewportToolBar.cpp
+++ b/Engine/Source/Editor/LevelEditor/Private/SLevelViewportToolBar.cpp
@@ -127,7 +127,7 @@ void SLevelViewportToolBar::Construct( const FArguments& InArgs )
 	[
 		SNew(SBorder)
 		.BorderImage(FAppStyle::Get().GetBrush("EditorViewportToolBar.Background"))
-		.Cursor(EMouseCursor::Default)
+		.Visibility(EVisibility::SelfHitTestInvisible)
 		[
 			SNew( SVerticalBox )
 			+ SVerticalBox::Slot()
diff --git a/Engine/Source/Editor/LevelEditor/Public/SLevelViewport.h b/Engine/Source/Editor/LevelEditor/Public/SLevelViewport.h
--- a/Engine/Source/Editor/LevelEditor/Public/SLevelViewport.h
+++ b/Engine/Source/Editor/LevelEditor/Public/SLevelViewport.h
@@ -397,7 +397,7 @@ public:
 	EVisibility GetToolbarVisibility() const;
 
 	/** Get the visibility for items considered to be part of the 'full' viewport toolbar */
-	EVisibility GetFullToolbarVisibility() const { return (bShowToolbarAndControls && bShowFullToolbar) ? EVisibility::Visible : EVisibility::Collapsed; }
+	EVisibility GetFullToolbarVisibility() const { return (bShowToolbarAndControls && bShowFullToolbar) ? EVisibility::SelfHitTestInvisible : EVisibility::Collapsed; }
 
 	/** Unpin and close all actor preview windows */
 	void RemoveAllPreviews(const bool bRemoveFromDesktopViewport = true);
