Fixes ActionableMessage widget not being hit testable anymore.

diff --git a/Engine/Source/Editor/ActionableMessage/Private/SActionableMessageViewportWidget.cpp b/Engine/Source/Editor/ActionableMessage/Private/SActionableMessageViewportWidget.cpp
--- a/Engine/Source/Editor/ActionableMessage/Private/SActionableMessageViewportWidget.cpp
+++ b/Engine/Source/Editor/ActionableMessage/Private/SActionableMessageViewportWidget.cpp
@@ -247,7 +247,7 @@ EVisibility SActionableMessageViewportWidget::GetVisibility()
 
 		if (CachedStateID == SubsystemStateID)
 		{
-			return ActionableMessages.IsEmpty() ? EVisibility::Collapsed : EVisibility::HitTestInvisible;
+			return ActionableMessages.IsEmpty() ? EVisibility::Collapsed : EVisibility::Visible;
 		}
 		
 		const TMap<FName, TSharedPtr<FActionableMessage>>& ActionableMessageSource = ActionableMessageSubsystem->GetActionableMessages();
@@ -266,7 +266,7 @@ EVisibility SActionableMessageViewportWidget::GetVisibility()
 			ActionableMessageList->RebuildList();
 		}
 		
-		return ActionableMessages.IsEmpty() ? EVisibility::Collapsed : EVisibility::HitTestInvisible;
+		return ActionableMessages.IsEmpty() ? EVisibility::Collapsed : EVisibility::Visible;
 	}
 
 	return EVisibility::Collapsed;
