Adds possibility for actionable message to force the widget's expansion

diff --git a/Engine/Source/Editor/ActionableMessage/Private/SActionableMessageViewportWidget.cpp b/Engine/Source/Editor/ActionableMessage/Private/SActionableMessageViewportWidget.cpp
--- a/Engine/Source/Editor/ActionableMessage/Private/SActionableMessageViewportWidget.cpp
+++ b/Engine/Source/Editor/ActionableMessage/Private/SActionableMessageViewportWidget.cpp
@@ -124,7 +124,7 @@ void SActionableMessageViewportWidget::Construct(const FArguments& InArgs)
 			.IsEnabled(true)
 			.ToolTipText_Lambda([this]()
 			{
-				if (!bExpanded)
+				if (!IsExpanded())
 				{
 					return LOCTEXT("ActionableMessages.CollapsedWarning", "Your project has warnings that need to be resolved. Click for details.");
 				}
@@ -157,7 +157,7 @@ void SActionableMessageViewportWidget::Construct(const FArguments& InArgs)
 					SAssignNew(TextBlock, STextBlock)
 					.Text_Lambda([this]()
 					{
-						if (bExpanded)
+						if (IsExpanded())
 						{
 							if (ActionableMessages.Num() == 1)
 							{
@@ -172,7 +172,7 @@ void SActionableMessageViewportWidget::Construct(const FArguments& InArgs)
 					.ColorAndOpacity(FStyleColors::Foreground)
 					.Font_Lambda([this] ()
 					{
-						if (bExpanded)
+						if (IsExpanded())
 						{
 							return FAppStyle::GetFontStyle(TEXT("BoldFont")); 
 						}
@@ -188,7 +188,7 @@ void SActionableMessageViewportWidget::Construct(const FArguments& InArgs)
 					{
 						FVector2d Size(0.f, 0.f);
 						
-						if (bExpanded && ActionableMessageList.IsValid() && TextBlock.IsValid())
+						if (IsExpanded() && ActionableMessageList.IsValid() && TextBlock.IsValid())
 						{
 							Size.X = ActionableMessageList->GetDesiredSize().X - TextBlock->GetDesiredSize().X - IconSize - 4.f * DefaultPadding - 2.f;
 						}
@@ -203,7 +203,7 @@ void SActionableMessageViewportWidget::Construct(const FArguments& InArgs)
 				.AutoWidth()
 				[
 					SNew(SImage)
-					.Image_Lambda([this]() { return bExpanded ? FAppStyle::GetBrush("ContentBrowser.SortDown") : FAppStyle::GetBrush("Symbols.LeftArrow"); })
+					.Image_Lambda([this]() { return IsExpanded() ? FAppStyle::GetBrush("ContentBrowser.SortDown") : FAppStyle::GetBrush("Symbols.LeftArrow"); })
 					.Visibility_Lambda([this]() { return ActionableMessages.IsEmpty() ? EVisibility::Collapsed : EVisibility::Visible; })
 					.ColorAndOpacity(FSlateColor::UseForeground())
 				]
@@ -215,7 +215,7 @@ void SActionableMessageViewportWidget::Construct(const FArguments& InArgs)
 		[
 			SNew(SSeparator)
 			.Thickness(1.f)
-			.Visibility_Lambda([this]() {return bExpanded ? EVisibility::Visible : EVisibility::Collapsed;})
+			.Visibility_Lambda([this]() {return IsExpanded() ? EVisibility::Visible : EVisibility::Collapsed;})
 		]
 		+ SVerticalBox::Slot()
 		.AutoHeight()
@@ -226,7 +226,7 @@ void SActionableMessageViewportWidget::Construct(const FArguments& InArgs)
 			.OnGenerateRow(this, &SActionableMessageViewportWidget::OnGenerateRow)
 			.Visibility_Lambda([this]()
 			{
-				return bExpanded ? EVisibility::Visible : EVisibility::Collapsed;
+				return IsExpanded() ? EVisibility::Visible : EVisibility::Collapsed;
 			})
 		]
 	];
@@ -253,12 +253,22 @@ EVisibility SActionableMessageViewportWidget::GetVisibility()
 		const TMap<FName, TSharedPtr<FActionableMessage>>& ActionableMessageSource = ActionableMessageSubsystem->GetActionableMessages();
 
 		ActionableMessages.Empty();
+		bForceExpand = false;
 
 		Algo::Transform(ActionableMessageSource, ActionableMessages, [](const TPair<FName, TSharedPtr<FActionableMessage>>& ActionableMessagePair)
 		{
 			return ActionableMessagePair.Value;
 		});
 
+		for (TSharedPtr<FActionableMessage> ActionableMessage : ActionableMessages)
+		{
+			if (ActionableMessage->bForceExpand)
+			{
+				bForceExpand = true;
+				break;
+			}
+		}
+
 		CachedStateID = SubsystemStateID;
 
 		if (ActionableMessageList.IsValid())
diff --git a/Engine/Source/Editor/ActionableMessage/Public/ActionableMessageSubsystem.h b/Engine/Source/Editor/ActionableMessage/Public/ActionableMessageSubsystem.h
--- a/Engine/Source/Editor/ActionableMessage/Public/ActionableMessageSubsystem.h
+++ b/Engine/Source/Editor/ActionableMessage/Public/ActionableMessageSubsystem.h
@@ -28,6 +28,11 @@ struct FActionableMessage
 	 * The callback associated with the action button, if any. 
 	 */
 	TFunction<void()> ActionCallback;
+
+	/**
+	 * If true, this message requires the panel to be expanded.
+	 */
+	bool bForceExpand = false;
 };
 
 FORCEINLINE bool operator!=(const FActionableMessage& LHS, const FActionableMessage& RHS);
diff --git a/Engine/Source/Editor/ActionableMessage/Public/SActionableMessageViewportWidget.h b/Engine/Source/Editor/ActionableMessage/Public/SActionableMessageViewportWidget.h
--- a/Engine/Source/Editor/ActionableMessage/Public/SActionableMessageViewportWidget.h
+++ b/Engine/Source/Editor/ActionableMessage/Public/SActionableMessageViewportWidget.h
@@ -35,10 +35,14 @@ public:
 
 	TSharedRef<ITableRow> OnGenerateRow(TSharedPtr<FActionableMessage> InActionableMessage, const TSharedRef<STableViewBase>& OwnerTable);
 
+private:
+	bool IsExpanded() const { return bExpanded || bForceExpand;	}
+
 private:
 	TSharedPtr<SListView<TSharedPtr<FActionableMessage>>> ActionableMessageList;
 	TSharedPtr<STextBlock> TextBlock;
 	TArray<TSharedPtr<FActionableMessage>> ActionableMessages;
 	uint32 CachedStateID = 0;
 	bool bExpanded = false;
+	bool bForceExpand = false;
 };
