Fix crash in actionable message system when reloading a level without clicking Update.
UI Widget wasn't noticing that the UActionableMessageSubsystem that provided data to it was being swapped out for a new instance when the editor loads a new level.  A stale pointer to the old level's subsystem was left attached to the button.

diff --git a/Engine/Source/Editor/ActionableMessage/Private/SActionableMessageViewportWidget.cpp b/Engine/Source/Editor/ActionableMessage/Private/SActionableMessageViewportWidget.cpp
--- a/Engine/Source/Editor/ActionableMessage/Private/SActionableMessageViewportWidget.cpp
+++ b/Engine/Source/Editor/ActionableMessage/Private/SActionableMessageViewportWidget.cpp
@@ -243,9 +243,10 @@ EVisibility SActionableMessageViewportWidget::GetVisibility()
 			return EVisibility::Collapsed;
 		}
 
+		TWeakObjectPtr<const UObject> DataSourceID(ActionableMessageSubsystem);  // detect if the data source has changed, probably from loading a new level
 		const uint32 SubsystemStateID = ActionableMessageSubsystem->GetStateID();
 
-		if (CachedStateID == SubsystemStateID)
+		if (CachedStateID == SubsystemStateID && CachedDataSourceID == DataSourceID)
 		{
 			return ActionableMessages.IsEmpty() ? EVisibility::Collapsed : EVisibility::Visible;
 		}
@@ -270,6 +271,7 @@ EVisibility SActionableMessageViewportWidget::GetVisibility()
 		}
 
 		CachedStateID = SubsystemStateID;
+		CachedDataSourceID = DataSourceID;
 
 		if (ActionableMessageList.IsValid())
 		{
diff --git a/Engine/Source/Editor/ActionableMessage/Public/SActionableMessageViewportWidget.h b/Engine/Source/Editor/ActionableMessage/Public/SActionableMessageViewportWidget.h
--- a/Engine/Source/Editor/ActionableMessage/Public/SActionableMessageViewportWidget.h
+++ b/Engine/Source/Editor/ActionableMessage/Public/SActionableMessageViewportWidget.h
@@ -43,6 +43,7 @@ private:
 	TSharedPtr<STextBlock> TextBlock;
 	TArray<TSharedPtr<FActionableMessage>> ActionableMessages;
 	uint32 CachedStateID = 0;
+	TWeakObjectPtr<const UObject> CachedDataSourceID;
 	bool bExpanded = false;
 	bool bForceExpand = false;
 };
